Experimental Windows Virtual Keyboard plugin

- primarily for handheld Windows devices
- not shown to be working yet due to Win11 bugs (fixes due in Win11 26100.5061, eta 26th Aug)

diff --git a/Engine/Plugins/Runtime/Windows/WinVirtualKeyboard/Source/WinVirtualKeyboard/Private/WinVirtualKeyboardModule.cpp b/Engine/Plugins/Runtime/Windows/WinVirtualKeyboard/Source/WinVirtualKeyboard/Private/WinVirtualKeyboardModule.cpp
new file mode 100644
--- /dev/null
+++ b/Engine/Plugins/Runtime/Windows/WinVirtualKeyboard/Source/WinVirtualKeyboard/Private/WinVirtualKeyboardModule.cpp
@@ -0,0 +1,200 @@
+// Copyright Epic Games, Inc. All Rights Reserved.
+
+#include "CoreMinimal.h"
+#include "Modules/ModuleInterface.h"
+#include "Modules/ModuleManager.h"
+#include "Features/IModularFeatures.h"
+#include "Framework/Application/IPlatformTextField.h"
+
+#if WITH_CPPWINRT
+
+THIRD_PARTY_INCLUDES_START
+#include "Microsoft/AllowMicrosoftPlatformTypes.h"
+#include <sdkddkver.h>
+#include <winrt/windows.ui.viewmanagement.core.h>
+#include "Microsoft/HideMicrosoftPlatformTypes.h"
+THIRD_PARTY_INCLUDES_END
+
+using namespace winrt::Windows::UI::ViewManagement::Core;
+
+#if defined(NTDDI_VERSION) && defined(NTDDI_WIN11_GE) && (NTDDI_VERSION > NTDDI_WIN11_GE) // definition was introduced in 10.0.26100.3624 so once we're past 26100, assume it is available
+	static const CoreInputViewKind DesiredInputKind = CoreInputViewKind::Gamepad;
+#else
+	static const CoreInputViewKind DesiredInputKind = (CoreInputViewKind)7;
+#endif
+
+
+#else
+
+THIRD_PARTY_INCLUDES_START
+#include "Microsoft/AllowMicrosoftPlatformTypes.h"
+#include "Microsoft/AllowMicrosoftPlatformAtomics.h"
+#include <sdkddkver.h>
+#include <roapi.h>
+#include <wrl.h>
+#include <windows.ui.viewmanagement.core.h>
+#include "Microsoft/HideMicrosoftPlatformAtomics.h"
+#include "Microsoft/HideMicrosoftPlatformTypes.h"
+THIRD_PARTY_INCLUDES_END
+
+using namespace ABI::Windows::UI::ViewManagement::Core;
+using namespace Microsoft::WRL;
+using namespace Microsoft::WRL::Wrappers;
+
+#if defined(NTDDI_VERSION) && defined(NTDDI_WIN11_GE) && (NTDDI_VERSION > NTDDI_WIN11_GE) // definition was introduced in 10.0.26100.3624 so once we're past 26100, assume it is available
+	static const CoreInputViewKind DesiredInputKind = CoreInputViewKind_Gamepad;
+#else
+	static const CoreInputViewKind DesiredInputKind = (CoreInputViewKind)7;
+#endif
+
+#endif // WITH_CPPWINRT
+
+
+
+class FWindowsPlatformTextField : public IPlatformTextField
+{
+public:
+
+#if WITH_CPPWINRT
+	FWindowsPlatformTextField()
+	{
+	}
+#else
+	FWindowsPlatformTextField(ComPtr<ICoreInputView3> InView3)
+		: View3(InView3)
+	{
+	}
+#endif // !WITH_CPPWINRT
+
+
+	virtual void ShowVirtualKeyboard(bool bShow, int32 UserIndex, TSharedPtr<IVirtualKeyboardEntry> TextEntryWidget) override
+	{
+#if WITH_CPPWINRT
+
+		bool bSuccess = false;
+		if (bShow)
+		{
+			bSuccess = CoreInputView::GetForCurrentView().TryShow(DesiredInputKind);
+		}
+		else
+		{
+			bSuccess = CoreInputView::GetForCurrentView().TryHide();
+		}
+		UE_CLOG(!bSuccess, PLATFORM_GLOBAL_LOG_CATEGORY, Error, TEXT("WinVirtualKeyboard %s failed"), bShow ? TEXT("show") : TEXT("hide"));
+
+#else
+		HRESULT hResult;
+		boolean Success;
+
+		if (bShow)
+		{
+			hResult = View3->TryShowWithKind(DesiredInputKind, &Success);
+		}
+		else
+		{
+			hResult = View3->TryHide(&Success);
+		}
+
+		bool bSuccess = SUCCEEDED(hResult) && Success;
+		UE_CLOG(!bSuccess, PLATFORM_GLOBAL_LOG_CATEGORY, Error, TEXT("WinVirtualKeyboard %s failed. hResult=0x%X, Result=%s"), bShow ? TEXT("show") : TEXT("hide"), hResult, Success ? TEXT("true") : TEXT("false") );
+#endif // WITH_CPPWINRT
+	}
+
+
+private:
+#if !WITH_CPPWINRT
+	ComPtr<ICoreInputView3> View3;
+#endif //!WITH_CPPWINRT
+
+};
+
+
+class FWindowsVirtualKeyboardModule : public IModuleInterface, public IPlatformTextFieldFactory
+{
+public:
+
+	// IModuleInterface
+	virtual void StartupModule() override
+	{
+		IModularFeatures::Get().RegisterModularFeature(IPlatformTextFieldFactory::FeatureName, this);
+	}
+
+	virtual void ShutdownModule() override
+	{
+		IModularFeatures::Get().UnregisterModularFeature(IPlatformTextFieldFactory::FeatureName, this);
+	}
+
+	
+	// IPlatformTextFieldFactory
+	virtual TUniquePtr<IPlatformTextField> CreateInstance() override
+	{
+#if WITH_CPPWINRT
+		return MakeUnique<FWindowsPlatformTextField>();
+
+#else
+		ComPtr<ICoreInputView3> View3 = GetCoreInputView3();
+		if (View3 == nullptr)
+		{
+			return nullptr;
+		}
+
+		return MakeUnique<FWindowsPlatformTextField>(View3);
+#endif //WITH_CPPWINRT
+	}
+
+private:
+
+#if !WITH_CPPWINRT
+	ComPtr<ICoreInputView3> GetCoreInputView3()
+	{
+		HRESULT hResult;
+
+		// get the view
+		// @note: RoGetActivationFactory requires Win8 or higher... should this be preceeded with LoadLibrary(combase.dll) / GetProcAddress() ?
+		ComPtr<ICoreInputViewStatics> ViewStatics;
+		hResult = RoGetActivationFactory( HStringReference(RuntimeClass_Windows_UI_ViewManagement_Core_CoreInputView).Get(),  __uuidof(ICoreInputViewStatics), &ViewStatics);
+		if (FAILED(hResult))
+		{
+			UE_LOG(PLATFORM_GLOBAL_LOG_CATEGORY, Warning, TEXT("WinVirtualKeyboard will be unavailable: can't get view statics 0x%X"), hResult);
+			return nullptr;
+		}
+
+		ComPtr<ICoreInputView> View;
+		hResult = ViewStatics->GetForCurrentView(&View);
+		if (FAILED(hResult))
+		{
+			UE_LOG(PLATFORM_GLOBAL_LOG_CATEGORY, Warning, TEXT("WinVirtualKeyboard will be unavailable: can't get current view 0x%X"), hResult);
+			return nullptr;
+		}
+
+		ComPtr<ICoreInputView5> View5;
+		hResult = View.As(&View5);
+		if (SUCCEEDED(hResult))
+		{
+			boolean IsSupported;
+			hResult = View5->IsKindSupported(DesiredInputKind, &IsSupported);
+			if (FAILED(hResult) || !IsSupported)
+			{
+				UE_LOG(PLATFORM_GLOBAL_LOG_CATEGORY, Warning, TEXT("WinVirtualKeyboard will be unavailable: input kind not supported or failed 0x%X (%s)"), hResult, IsSupported ? TEXT("supported") : TEXT("not supported") );
+				return nullptr;
+			}
+		}
+
+		ComPtr<ICoreInputView3> View3;
+		hResult = View.As(&View3);
+		if (FAILED(hResult))
+		{
+			UE_LOG(PLATFORM_GLOBAL_LOG_CATEGORY, Warning, TEXT("WinVirtualKeyboard will be unavailable: can't get view3 0x%X"), hResult);
+			return nullptr;
+		}
+
+		return View3;
+	}
+#endif //!WITH_CPPWINRT
+
+};
+
+
+IMPLEMENT_MODULE(FWindowsVirtualKeyboardModule, WinVirtualKeyboard);
+
+
diff --git a/Engine/Plugins/Runtime/Windows/WinVirtualKeyboard/Source/WinVirtualKeyboard/WinVirtualKeyboard.Build.cs b/Engine/Plugins/Runtime/Windows/WinVirtualKeyboard/Source/WinVirtualKeyboard/WinVirtualKeyboard.Build.cs
new file mode 100644
--- /dev/null
+++ b/Engine/Plugins/Runtime/Windows/WinVirtualKeyboard/Source/WinVirtualKeyboard/WinVirtualKeyboard.Build.cs
@@ -0,0 +1,38 @@
+// Copyright Epic Games, Inc. All Rights Reserved.
+
+using System.IO;
+
+namespace UnrealBuildTool.Rules
+{
+	public class WinVirtualKeyboard : ModuleRules
+	{
+		// for A/B testing the two approaches - COM vs CPPWinRT.
+		static readonly bool bUseCOM = true;
+
+		public WinVirtualKeyboard(ReadOnlyTargetRules Target) : base(Target)
+		{
+			PrivateDependencyModuleNames.Add("Slate");
+			PublicDependencyModuleNames.Add("Core");
+			PublicSystemLibraries.Add("WindowsApp.lib");
+
+			if (bUseCOM)
+			{
+				// use WinRT ABI via COM
+				PrivateDefinitions.Add("WITH_CPPWINRT=0");
+			}
+			else
+			{
+				PrivateDefinitions.Add("WITH_CPPWINRT=1");
+
+				// CPPWinRT requires /Ehsc
+				bEnableExceptions = true;
+
+				// add the CPPWinRT include path if the project isn't using CPPWinRT
+				if (!Target.WindowsPlatform.bUseCPPWinRT)
+				{
+					PrivateIncludePaths.Add( Path.Combine(Target.WindowsPlatform.WindowsSdkDir, "include", Target.WindowsPlatform.WindowsSdkVersion.ToString(), "cppwinrt") );
+				}
+			}
+		}
+	}
+}
diff --git a/Engine/Plugins/Runtime/Windows/WinVirtualKeyboard/WinVirtualKeyboard.uplugin b/Engine/Plugins/Runtime/Windows/WinVirtualKeyboard/WinVirtualKeyboard.uplugin
new file mode 100644
--- /dev/null
+++ b/Engine/Plugins/Runtime/Windows/WinVirtualKeyboard/WinVirtualKeyboard.uplugin
@@ -0,0 +1,27 @@
+{
+	"FileVersion": 3,
+	"Version": 1,
+	"VersionName": "0.1",
+	"FriendlyName": "Windows Virtual Keyboard",
+	"Description": "Virtual Keyboard support for Windows. Requires Windows 11 26100.5061 or later",
+	"Category": "Input Devices",
+	"CreatedBy": "Epic Games, Inc.",
+	"CreatedByURL": "https://epicgames.com",
+	"DocsURL": "",
+	"MarketplaceURL": "",
+	"SupportURL": "",
+	"EnabledByDefault": false,
+	"CanContainContent": false,
+	"IsBetaVersion": true,
+	"Installed": false,
+	"RequiresBuildPlatform": false,
+	"SupportedTargetPlatforms": [ "Win64" ],
+	"Modules": [
+		{
+			"Name": "WinVirtualKeyboard",
+			"Type": "RuntimeNoCommandlet",
+			"LoadingPhase": "Default",
+			"PlatformAllowList": [ "Win64" ]
+		}
+	]
+}
\ No newline at end of file
