static analysis fixes

#rnx

diff --git a/Engine/Plugins/Runtime/Windows/WinVirtualKeyboard/Source/WinVirtualKeyboard/Private/WinVirtualKeyboardModule.cpp b/Engine/Plugins/Runtime/Windows/WinVirtualKeyboard/Source/WinVirtualKeyboard/Private/WinVirtualKeyboardModule.cpp
--- a/Engine/Plugins/Runtime/Windows/WinVirtualKeyboard/Source/WinVirtualKeyboard/Private/WinVirtualKeyboardModule.cpp
+++ b/Engine/Plugins/Runtime/Windows/WinVirtualKeyboard/Source/WinVirtualKeyboard/Private/WinVirtualKeyboardModule.cpp
@@ -84,7 +84,7 @@ public:
 
 #else
 		HRESULT hResult;
-		boolean Success;
+		boolean Success = 0;
 
 		if (bShow)
 		{
@@ -171,7 +171,7 @@ private:
 		hResult = View.As(&View5);
 		if (SUCCEEDED(hResult))
 		{
-			boolean IsSupported;
+			boolean IsSupported = 0;
 			hResult = View5->IsKindSupported(DesiredInputKind, &IsSupported);
 			if (FAILED(hResult) || !IsSupported)
 			{
