remove incorrect code block

#rnx

diff --git a/Engine/Plugins/Runtime/Windows/WinVirtualKeyboard/Source/WinVirtualKeyboard/Private/WinVirtualKeyboardModule.cpp b/Engine/Plugins/Runtime/Windows/WinVirtualKeyboard/Source/WinVirtualKeyboard/Private/WinVirtualKeyboardModule.cpp
--- a/Engine/Plugins/Runtime/Windows/WinVirtualKeyboard/Source/WinVirtualKeyboard/Private/WinVirtualKeyboardModule.cpp
+++ b/Engine/Plugins/Runtime/Windows/WinVirtualKeyboard/Source/WinVirtualKeyboard/Private/WinVirtualKeyboardModule.cpp
@@ -167,19 +167,6 @@ private:
 			return nullptr;
 		}
 
-		ComPtr<ICoreInputView5> View5;
-		hResult = View.As(&View5);
-		if (SUCCEEDED(hResult))
-		{
-			boolean IsSupported = 0;
-			hResult = View5->IsKindSupported(DesiredInputKind, &IsSupported);
-			if (FAILED(hResult) || !IsSupported)
-			{
-				UE_LOG(PLATFORM_GLOBAL_LOG_CATEGORY, Warning, TEXT("WinVirtualKeyboard will be unavailable: input kind not supported or failed 0x%X (%s)"), hResult, IsSupported ? TEXT("supported") : TEXT("not supported") );
-				return nullptr;
-			}
-		}
-
 		ComPtr<ICoreInputView3> View3;
 		hResult = View.As(&View3);
 		if (FAILED(hResult))
