disallow WinVirtualKeyboard on servers

#rnx

diff --git a/Engine/Plugins/Runtime/Windows/WinVirtualKeyboard/WinVirtualKeyboard.uplugin b/Engine/Plugins/Runtime/Windows/WinVirtualKeyboard/WinVirtualKeyboard.uplugin
--- a/Engine/Plugins/Runtime/Windows/WinVirtualKeyboard/WinVirtualKeyboard.uplugin
+++ b/Engine/Plugins/Runtime/Windows/WinVirtualKeyboard/WinVirtualKeyboard.uplugin
@@ -21,7 +21,8 @@
 			"Name": "WinVirtualKeyboard",
 			"Type": "RuntimeNoCommandlet",
 			"LoadingPhase": "Default",
-			"PlatformAllowList": [ "Win64" ]
+			"PlatformAllowList": [ "Win64" ],
+			"TargetDenyList": [ "Server" ]
 		}
 	]
 }
\ No newline at end of file
