[WinVirtualKeyboard] Implemented GenericApplication Platform API callbacks

#rnx

#Synced-CL 49981045

diff --git a/Engine/Plugins/Runtime/Windows/WinVirtualKeyboard/Source/WinVirtualKeyboard/Private/WinVirtualKeyboardModule.cpp b/Engine/Plugins/Runtime/Windows/WinVirtualKeyboard/Source/WinVirtualKeyboard/Private/WinVirtualKeyboardModule.cpp
--- a/Engine/Plugins/Runtime/Windows/WinVirtualKeyboard/Source/WinVirtualKeyboard/Private/WinVirtualKeyboardModule.cpp
+++ b/Engine/Plugins/Runtime/Windows/WinVirtualKeyboard/Source/WinVirtualKeyboard/Private/WinVirtualKeyboardModule.cpp
@@ -5,6 +5,7 @@
 #include "Modules/ModuleManager.h"
 #include "Features/IModularFeatures.h"
 #include "Framework/Application/IPlatformTextField.h"
+#include "Framework/Application/SlateApplication.h"
 
 #if WITH_CPPWINRT
 
@@ -70,7 +71,6 @@ public:
 	virtual void ShowVirtualKeyboard(bool bShow, int32 UserIndex, TSharedPtr<IVirtualKeyboardEntry> TextEntryWidget) override
 	{
 #if WITH_CPPWINRT
-
 		bool bSuccess = false;
 		if (bShow)
 		{
@@ -96,10 +96,28 @@ public:
 		}
 
 		bool bSuccess = SUCCEEDED(hResult) && Success;
-		UE_CLOG(!bSuccess, PLATFORM_GLOBAL_LOG_CATEGORY, Error, TEXT("WinVirtualKeyboard %s failed. hResult=0x%X, Result=%s"), bShow ? TEXT("show") : TEXT("hide"), hResult, Success ? TEXT("true") : TEXT("false") );
+		UE_CLOG(!bSuccess, PLATFORM_GLOBAL_LOG_CATEGORY, Error, TEXT("WinVirtualKeyboard %s failed. hResult=0x%X, Result=%s"), bShow ? TEXT("show") : TEXT("hide"), hResult, Success ? TEXT("true") : TEXT("false"));
+
 #endif // WITH_CPPWINRT
-	}
 
+		if (bSuccess)
+		{
+			FSlateApplication& SlateApp = FSlateApplication::Get();
+			TSharedPtr<GenericApplication> PlatformAppPtr = SlateApp.GetPlatformApplication();
+			if (GenericApplication* GenericApplication = PlatformAppPtr.Get())
+			{
+				if (bShow)
+				{
+					FPlatformRect DummyRect{0,0,0,0}; // TODO: Find the way to get VK Rect
+					GenericApplication->OnVirtualKeyboardShown().Broadcast(DummyRect);
+				}
+				else
+				{
+					GenericApplication->OnVirtualKeyboardHidden().Broadcast();
+				}
+			}
+		}
+	}
 
 private:
 #if !WITH_CPPWINRT
diff --git a/Engine/Plugins/Runtime/Windows/WinVirtualKeyboard/Source/WinVirtualKeyboard/WinVirtualKeyboard.Build.cs b/Engine/Plugins/Runtime/Windows/WinVirtualKeyboard/Source/WinVirtualKeyboard/WinVirtualKeyboard.Build.cs
--- a/Engine/Plugins/Runtime/Windows/WinVirtualKeyboard/Source/WinVirtualKeyboard/WinVirtualKeyboard.Build.cs
+++ b/Engine/Plugins/Runtime/Windows/WinVirtualKeyboard/Source/WinVirtualKeyboard/WinVirtualKeyboard.Build.cs
@@ -11,6 +11,7 @@ namespace UnrealBuildTool.Rules
 
 		public WinVirtualKeyboard(ReadOnlyTargetRules Target) : base(Target)
 		{
+			PrivateDependencyModuleNames.Add("SlateCore");
 			PrivateDependencyModuleNames.Add("Slate");
 			PublicDependencyModuleNames.Add("Core");
 			PublicSystemLibraries.Add("WindowsApp.lib");
