remove unused cppwinrt implementation from WinVirtualKeyboard

#rnx

#Synced-CL 50642529

diff --git a/Engine/Plugins/Runtime/Windows/WinVirtualKeyboard/Source/WinVirtualKeyboard/Private/WinVirtualKeyboardModule.cpp b/Engine/Plugins/Runtime/Windows/WinVirtualKeyboard/Source/WinVirtualKeyboard/Private/WinVirtualKeyboardModule.cpp
--- a/Engine/Plugins/Runtime/Windows/WinVirtualKeyboard/Source/WinVirtualKeyboard/Private/WinVirtualKeyboardModule.cpp
+++ b/Engine/Plugins/Runtime/Windows/WinVirtualKeyboard/Source/WinVirtualKeyboard/Private/WinVirtualKeyboardModule.cpp
@@ -7,26 +7,6 @@
 #include "Framework/Application/IPlatformTextField.h"
 #include "Framework/Application/SlateApplication.h"
 
-#if WITH_CPPWINRT
-
-THIRD_PARTY_INCLUDES_START
-#include "Microsoft/AllowMicrosoftPlatformTypes.h"
-#include <sdkddkver.h>
-#include <winrt/windows.ui.viewmanagement.core.h>
-#include "Microsoft/HideMicrosoftPlatformTypes.h"
-THIRD_PARTY_INCLUDES_END
-
-using namespace winrt::Windows::UI::ViewManagement::Core;
-
-#if defined(NTDDI_VERSION) && defined(NTDDI_WIN11_GE) && (NTDDI_VERSION > NTDDI_WIN11_GE) // definition was introduced in 10.0.26100.3624 so once we're past 26100, assume it is available
-	static const CoreInputViewKind DesiredInputKind = CoreInputViewKind::Gamepad;
-#else
-	static const CoreInputViewKind DesiredInputKind = (CoreInputViewKind)7;
-#endif
-
-
-#else
-
 THIRD_PARTY_INCLUDES_START
 #include "Microsoft/AllowMicrosoftPlatformTypes.h"
 #include "Microsoft/AllowMicrosoftPlatformAtomics.h"
@@ -48,41 +28,19 @@ using namespace Microsoft::WRL::Wrappers;
 	static const CoreInputViewKind DesiredInputKind = (CoreInputViewKind)7;
 #endif
 
-#endif // WITH_CPPWINRT
-
-
 
 class FWindowsPlatformTextField : public IPlatformTextField
 {
 public:
 
-#if WITH_CPPWINRT
-	FWindowsPlatformTextField()
-	{
-	}
-#else
 	FWindowsPlatformTextField(ComPtr<ICoreInputView3> InView3)
 		: View3(InView3)
 	{
 	}
-#endif // !WITH_CPPWINRT
 
 
 	virtual void ShowVirtualKeyboard(bool bShow, int32 UserIndex, TSharedPtr<IVirtualKeyboardEntry> TextEntryWidget) override
 	{
-#if WITH_CPPWINRT
-		bool bSuccess = false;
-		if (bShow)
-		{
-			bSuccess = CoreInputView::GetForCurrentView().TryShow(DesiredInputKind);
-		}
-		else
-		{
-			bSuccess = CoreInputView::GetForCurrentView().TryHide();
-		}
-		UE_CLOG(!bSuccess, PLATFORM_GLOBAL_LOG_CATEGORY, Error, TEXT("WinVirtualKeyboard %s failed"), bShow ? TEXT("show") : TEXT("hide"));
-
-#else
 		HRESULT hResult;
 		boolean Success = 0;
 
@@ -98,8 +56,6 @@ public:
 		bool bSuccess = SUCCEEDED(hResult) && Success;
 		UE_CLOG(!bSuccess, PLATFORM_GLOBAL_LOG_CATEGORY, Error, TEXT("WinVirtualKeyboard %s failed. hResult=0x%X, Result=%s"), bShow ? TEXT("show") : TEXT("hide"), hResult, Success ? TEXT("true") : TEXT("false"));
 
-#endif // WITH_CPPWINRT
-
 		if (bSuccess)
 		{
 			FSlateApplication& SlateApp = FSlateApplication::Get();
@@ -120,10 +76,7 @@ public:
 	}
 
 private:
-#if !WITH_CPPWINRT
 	ComPtr<ICoreInputView3> View3;
-#endif //!WITH_CPPWINRT
-
 };
 
 
@@ -146,10 +99,6 @@ public:
 	// IPlatformTextFieldFactory
 	virtual TUniquePtr<IPlatformTextField> CreateInstance() override
 	{
-#if WITH_CPPWINRT
-		return MakeUnique<FWindowsPlatformTextField>();
-
-#else
 		ComPtr<ICoreInputView3> View3 = GetCoreInputView3();
 		if (View3 == nullptr)
 		{
@@ -157,12 +106,10 @@ public:
 		}
 
 		return MakeUnique<FWindowsPlatformTextField>(View3);
-#endif //WITH_CPPWINRT
 	}
 
 private:
 
-#if !WITH_CPPWINRT
 	ComPtr<ICoreInputView3> GetCoreInputView3()
 	{
 		HRESULT hResult;
@@ -195,8 +142,6 @@ private:
 
 		return View3;
 	}
-#endif //!WITH_CPPWINRT
-
 };
 
 
diff --git a/Engine/Plugins/Runtime/Windows/WinVirtualKeyboard/Source/WinVirtualKeyboard/WinVirtualKeyboard.Build.cs b/Engine/Plugins/Runtime/Windows/WinVirtualKeyboard/Source/WinVirtualKeyboard/WinVirtualKeyboard.Build.cs
--- a/Engine/Plugins/Runtime/Windows/WinVirtualKeyboard/Source/WinVirtualKeyboard/WinVirtualKeyboard.Build.cs
+++ b/Engine/Plugins/Runtime/Windows/WinVirtualKeyboard/Source/WinVirtualKeyboard/WinVirtualKeyboard.Build.cs
@@ -6,34 +6,12 @@ namespace UnrealBuildTool.Rules
 {
 	public class WinVirtualKeyboard : ModuleRules
 	{
-		// for A/B testing the two approaches - COM vs CPPWinRT.
-		static readonly bool bUseCOM = true;
-
 		public WinVirtualKeyboard(ReadOnlyTargetRules Target) : base(Target)
 		{
 			PrivateDependencyModuleNames.Add("SlateCore");
 			PrivateDependencyModuleNames.Add("Slate");
 			PublicDependencyModuleNames.Add("Core");
 			PublicSystemLibraries.Add("WindowsApp.lib");
-
-			if (bUseCOM)
-			{
-				// use WinRT ABI via COM
-				PrivateDefinitions.Add("WITH_CPPWINRT=0");
-			}
-			else
-			{
-				PrivateDefinitions.Add("WITH_CPPWINRT=1");
-
-				// CPPWinRT requires /Ehsc
-				bEnableExceptions = true;
-
-				// add the CPPWinRT include path if the project isn't using CPPWinRT
-				if (!Target.WindowsPlatform.bUseCPPWinRT)
-				{
-					PrivateIncludePaths.Add( Path.Combine(Target.WindowsPlatform.WindowsSdkDir, "include", Target.WindowsPlatform.WindowsSdkVersion.ToString(), "cppwinrt") );
-				}
-			}
 		}
 	}
 }
